Fix CIS warning.

#rnx

diff --git a/Engine/Source/Programs/PixelStreaming/WebRTCProxy/WebRTCProxy.Target.cs b/Engine/Source/Programs/PixelStreaming/WebRTCProxy/WebRTCProxy.Target.cs
--- a/Engine/Source/Programs/PixelStreaming/WebRTCProxy/WebRTCProxy.Target.cs
+++ b/Engine/Source/Programs/PixelStreaming/WebRTCProxy/WebRTCProxy.Target.cs
@@ -13,7 +13,7 @@ public class WebRTCProxyTarget : TargetRules
 		GlobalDefinitions.Add("INCL_EXTRA_HTON_FUNCTIONS");
 
 		// Lean and mean
-		bCompileLeanAndMeanUE = true;
+		bBuildDeveloperTools = false;
 		
         // No editor needed
         bBuildEditor = false;
