Use scoped Windows TRUE and FALSE constants.

#rnx

diff --git a/Engine/Source/Programs/PixelStreaming/WebRTCProxy/src/Console.cpp b/Engine/Source/Programs/PixelStreaming/WebRTCProxy/src/Console.cpp
--- a/Engine/Source/Programs/PixelStreaming/WebRTCProxy/src/Console.cpp
+++ b/Engine/Source/Programs/PixelStreaming/WebRTCProxy/src/Console.cpp
@@ -27,7 +27,7 @@ void FConsole::Init(short Width, short Height, short BufferWidth, short BufferHe
 
 	// allocate a console for this app.
 	// NOTE: It fails if there is a console already
-	bOwnsConsole = AllocConsole() == TRUE ? true : false;
+	bOwnsConsole = AllocConsole() == Windows::TRUE ? true : false;
 
 	hConsoleHandle = GetStdHandle(STD_OUTPUT_HANDLE);
 	// set the screen buffer to be big enough to let us scroll text
@@ -39,7 +39,7 @@ void FConsole::Init(short Width, short Height, short BufferWidth, short BufferHe
 	// Set the real window size (need to be smaller than the buffer
 	ConInfo.srWindow.Bottom = Height - 1;
 	ConInfo.srWindow.Right = Width - 1;
-	SetConsoleWindowInfo(hConsoleHandle, TRUE, &ConInfo.srWindow);
+	SetConsoleWindowInfo(hConsoleHandle, Windows::TRUE, &ConInfo.srWindow);
 	Center();
 	EnableUTF8Support();
 	SetTextColour(EColour::White);
@@ -69,7 +69,7 @@ void FConsole::Center()
 	    0,
 	    SWP_NOSIZE | SWP_NOZORDER);
 
-	SwitchToThisWindow(hConsoleWnd, TRUE);
+	SwitchToThisWindow(hConsoleWnd, Windows::TRUE);
 }
 
 void FConsole::EnableUTF8Support()
diff --git a/Engine/Source/Programs/PixelStreaming/WebRTCProxy/src/FileLogOutput.h b/Engine/Source/Programs/PixelStreaming/WebRTCProxy/src/FileLogOutput.h
--- a/Engine/Source/Programs/PixelStreaming/WebRTCProxy/src/FileLogOutput.h
+++ b/Engine/Source/Programs/PixelStreaming/WebRTCProxy/src/FileLogOutput.h
@@ -28,9 +28,9 @@ public:
 	/**
 	* @param bAutoNewLine if true, ever Log call will automatically append a '\n' to the message
 	*/
-	void SetAutoNewLine(bool bAutoNewLine)
+	void SetAutoNewLine(bool bInAutoNewLine)
 	{
-		this->bAutoNewLine = bAutoNewLine;
+		this->bAutoNewLine = bInAutoNewLine;
 	}
 
 	void Write(const char* Msg);
diff --git a/Engine/Source/Programs/PixelStreaming/WebRTCProxy/src/WebRTCProxy.cpp b/Engine/Source/Programs/PixelStreaming/WebRTCProxy/src/WebRTCProxy.cpp
--- a/Engine/Source/Programs/PixelStreaming/WebRTCProxy/src/WebRTCProxy.cpp
+++ b/Engine/Source/Programs/PixelStreaming/WebRTCProxy/src/WebRTCProxy.cpp
@@ -11,7 +11,7 @@
 #include "CrashDetection.h"
 
 const char* Help =
-"\
+    "\
 WebRTCProxy\n\
 Copyright 1998-2019 Epic Games, Inc. All Rights Reserved.\n\
 \n\
@@ -53,12 +53,12 @@ If specified, it will use local time in logging, instead of UTC.\n\
 TCHAR GInternalProjectName[64] = TEXT("WebRTCProxy");
 IMPLEMENT_FOREIGN_ENGINE_DIR();
 
-std::pair<std::string, uint16_t> PARAM_Cirrus{ "127.0.0.1", 8888 };
+std::pair<std::string, uint16_t> PARAM_Cirrus{"127.0.0.1", 8888};
 uint16_t PARAM_UE4Port = 8124;
 bool PARAM_PlanB = false;
 bool PARAM_DbgWindow_Proxy = true;
 bool PARAM_DbgWindow_WebRTC = true;
-bool PARAM_LocalTime = false; // By default we use UTC time
+bool PARAM_LocalTime = false;  // By default we use UTC time
 
 bool ParseParameters(int argc, char* argv[])
 {
@@ -76,8 +76,7 @@ bool ParseParameters(int argc, char* argv[])
 	}
 
 	// Splits a string in the form of "XXXX:NNN" into a pair
-	auto ProcessAddressParameter = [&Params](const char* Name, std::pair<std::string, uint16_t>& OutAddr) -> bool
-	{
+	auto ProcessAddressParameter = [&Params](const char* Name, std::pair<std::string, uint16_t>& OutAddr) -> bool {
 		if (!Params.Has(Name))
 		{
 			return true;
@@ -86,17 +85,17 @@ bool ParseParameters(int argc, char* argv[])
 		const char* const Param = Params.Get(Name).c_str();
 		const char* Ptr = Param;
 		// Find separator
-		while (!(*Ptr==0 || *Ptr == ':' || *Ptr == '|'))
+		while (!(*Ptr == 0 || *Ptr == ':' || *Ptr == '|'))
 		{
 			Ptr++;
 		}
 
 		OutAddr.first = std::string(Param, Ptr);
 		// If at the end of the string, then no separator was found (and no port specified)
-		if (*Ptr && OutAddr.first!="")
+		if (*Ptr && OutAddr.first != "")
 		{
 			int Port = std::atoi(Ptr + 1);
-			if (Port < 1 || Port>65535)
+			if (Port < 1 || Port > 65535)
 			{
 				EG_LOG(LogDefault, Error, "Invalid port number for parameter '%s'", Name);
 				return false;
@@ -184,7 +183,7 @@ static BOOL WINAPI ConsoleCtrlHandler(DWORD dwCtrlType)
 
 	if (!MainThreadId)
 	{
-		return FALSE;
+		return Windows::FALSE;
 	}
 
 	PostThreadMessage(MainThreadId, WM_QUIT, 0, 0);
@@ -196,7 +195,7 @@ static BOOL WINAPI ConsoleCtrlHandler(DWORD dwCtrlType)
 
 	// Return TRUE if handled this message, further handler functions won't be called.
 	// Return FALSE to pass this message to further handlers until default handler calls ExitProcess().
-	return FALSE;
+	return Windows::FALSE;
 }
 
 int mainImpl(int argc, char* argv[])
@@ -205,7 +204,7 @@ int mainImpl(int argc, char* argv[])
 	Console.Init(120, 40, 400, 2000);
 
 	MainThreadId = GetCurrentThreadId();
-	SetConsoleCtrlHandler(ConsoleCtrlHandler, TRUE);
+	SetConsoleCtrlHandler(ConsoleCtrlHandler, Windows::TRUE);
 
 	// NOTE: Parsing the parameters before creating the file logger, so the log
 	// filename takes into account the -LocalTime parameter (if specified)
@@ -217,11 +216,14 @@ int mainImpl(int argc, char* argv[])
 	//
 	// Create file loggers
 	//
-	FFileLogOutput FileLogger(nullptr); // Our own log file
+	FFileLogOutput FileLogger(nullptr);  // Our own log file
 	// WebRTC logging
 	InitializeWebRTCLogging(rtc::LoggingSeverity::LS_VERBOSE);
 	// Make sure we stop the webrtc logging, otherwise it crashes on exit
-	SCOPE_EXIT{ StopWebRTCLogging(); };
+	SCOPE_EXIT
+	{
+		StopWebRTCLogging();
+	};
 
 	// Log the command line parameters, so we know what parameters were used for this run
 	{
@@ -252,7 +254,7 @@ int mainImpl(int argc, char* argv[])
 	rtc::Win32SocketServer w32_ss;
 	rtc::Win32Thread w32_thread(&w32_ss);
 	rtc::ThreadManager::Instance()->SetCurrentThread(&w32_thread);
-#elif EG_PLATFORM_LINUX==EG_PLATFORM_LINUX
+#elif EG_PLATFORM_LINUX == EG_PLATFORM_LINUX
 #error Not yet implemented
 #else
 #error Unknown platform
@@ -274,7 +276,7 @@ int mainImpl(int argc, char* argv[])
 
 	EG_LOG(LogDefault, Log, "Exiting UE4WebRTCProxy");
 
-    return EXIT_SUCCESS;
+	return EXIT_SUCCESS;
 }
 
 int main(int argc, char* argv[])
@@ -284,7 +286,7 @@ int main(int argc, char* argv[])
 	{
 		ExitCode = mainImpl(argc, argv);
 	}
-	catch (std::exception&e)
+	catch (std::exception& e)
 	{
 		printf("%s\n", e.what());
 		ExitCode = EXIT_FAILURE;
