Change Show FPS to button since we don't have the current shown state and so can't update the UI appropriately to it's current state.

diff --git a/Engine/Source/Programs/PixelStreaming/WebServers/SignallingWebServer/scripts/app.js b/Engine/Source/Programs/PixelStreaming/WebServers/SignallingWebServer/scripts/app.js
--- a/Engine/Source/Programs/PixelStreaming/WebServers/SignallingWebServer/scripts/app.js
+++ b/Engine/Source/Programs/PixelStreaming/WebServers/SignallingWebServer/scripts/app.js
@@ -91,9 +91,9 @@ function setupHtmlEvents(){
 	    }
 	}
 
-	let showFPSCheckBox = document.getElementById('show-fps-tgl');
-	if (showFPSCheckBox != null) {
-	    showFPSCheckBox.onchange = function (event) {
+	let showFPSButton = document.getElementById('show-fps-button');
+	if (showFPSButton != null) {
+	    showFPSButton.onclick = function (event) {
 	        let consoleDescriptor = {
 	            Console: 'stat fps'
 	        }
