Fix shadow variable warnings in WebRTCProxy.

diff --git a/Engine/Source/Programs/PixelStreaming/WebRTCProxy/src/ClientSession.cpp b/Engine/Source/Programs/PixelStreaming/WebRTCProxy/src/ClientSession.cpp
--- a/Engine/Source/Programs/PixelStreaming/WebRTCProxy/src/ClientSession.cpp
+++ b/Engine/Source/Programs/PixelStreaming/WebRTCProxy/src/ClientSession.cpp
@@ -116,10 +116,10 @@ void FClientSession::OnRemoveStream(rtc::scoped_refptr<webrtc::MediaStreamInterf
 	EG_LOG(LogDefault, Log, "%s : ClientId=%u, Stream=%s", __FUNCTION__, ClientId, Stream->id().c_str());
 }
 
-void FClientSession::OnDataChannel(rtc::scoped_refptr<webrtc::DataChannelInterface> DataChannel)
+void FClientSession::OnDataChannel(rtc::scoped_refptr<webrtc::DataChannelInterface> InDataChannel)
 {
 	EG_LOG(LogDefault, Log, "%s : ClientId=%u", __FUNCTION__, ClientId);
-	this->DataChannel = DataChannel;
+	this->DataChannel = InDataChannel;
 	this->DataChannel->RegisterObserver(this);
 }
 
diff --git a/Engine/Source/Programs/PixelStreaming/WebRTCProxy/src/CmdLine.cpp b/Engine/Source/Programs/PixelStreaming/WebRTCProxy/src/CmdLine.cpp
--- a/Engine/Source/Programs/PixelStreaming/WebRTCProxy/src/CmdLine.cpp
+++ b/Engine/Source/Programs/PixelStreaming/WebRTCProxy/src/CmdLine.cpp
@@ -6,9 +6,9 @@
 
 std::string FCmdLine::Empty;
 
-bool FCmdLine::Parse(int Argc, char* Argv[], bool CaseSensitive)
+bool FCmdLine::Parse(int Argc, char* Argv[], bool InCaseSensitive)
 {
-	this->CaseSensitive = CaseSensitive;
+	this->CaseSensitive = InCaseSensitive;
 
 	if (Argc <= 1)
 	{
diff --git a/Engine/Source/Programs/PixelStreaming/WebRTCProxy/src/Conductor.cpp b/Engine/Source/Programs/PixelStreaming/WebRTCProxy/src/Conductor.cpp
--- a/Engine/Source/Programs/PixelStreaming/WebRTCProxy/src/Conductor.cpp
+++ b/Engine/Source/Programs/PixelStreaming/WebRTCProxy/src/Conductor.cpp
@@ -135,17 +135,17 @@ void FConductor::AddStreams(FClientId ClientId)
 		{
 			Stream = PeerConnectionFactory->CreateLocalMediaStream(StreamId);
 
-			rtc::scoped_refptr<webrtc::AudioTrackInterface> AudioTrack(
+			rtc::scoped_refptr<webrtc::AudioTrackInterface> LocalAudioTrack(
 			    PeerConnectionFactory->CreateAudioTrack(AudioLabel, PeerConnectionFactory->CreateAudioSource(NULL)));
 
-			Stream->AddTrack(AudioTrack);
+			Stream->AddTrack(LocalAudioTrack);
 
 			std::unique_ptr<FNetworkVideoCapturer> VideoCapturerStrong = std::make_unique<FNetworkVideoCapturer>();
 			VideoCapturer = VideoCapturerStrong.get();
-			rtc::scoped_refptr<webrtc::VideoTrackInterface> VideoTrack(PeerConnectionFactory->CreateVideoTrack(
+			rtc::scoped_refptr<webrtc::VideoTrackInterface> LocalVideoTrack(PeerConnectionFactory->CreateVideoTrack(
 			    VideoLabel, PeerConnectionFactory->CreateVideoSource(std::move(VideoCapturerStrong))));
 
-			Stream->AddTrack(VideoTrack);
+			Stream->AddTrack(LocalVideoTrack);
 
 			typedef std::pair<std::string, rtc::scoped_refptr<webrtc::MediaStreamInterface> > MediaStreamPair;
 			Streams.insert(MediaStreamPair(Stream->id(), Stream));
diff --git a/Engine/Source/Programs/PixelStreaming/WebRTCProxy/src/VideoEncoder.cpp b/Engine/Source/Programs/PixelStreaming/WebRTCProxy/src/VideoEncoder.cpp
--- a/Engine/Source/Programs/PixelStreaming/WebRTCProxy/src/VideoEncoder.cpp
+++ b/Engine/Source/Programs/PixelStreaming/WebRTCProxy/src/VideoEncoder.cpp
@@ -111,9 +111,9 @@ int32_t FVideoEncoder::InitEncode(const webrtc::VideoCodec* CodecSettings, int32
 	return 0;
 }
 
-int32_t FVideoEncoder::RegisterEncodeCompleteCallback(webrtc::EncodedImageCallback* Callback)
+int32_t FVideoEncoder::RegisterEncodeCompleteCallback(webrtc::EncodedImageCallback* InCallback)
 {
-	this->Callback = Callback;
+	this->Callback = InCallback;
 	return 0;
 }
 
