Copy of 6989699 to from //UE4/Dev-Enterprise:

Allow the address (as well as the port) of UE4 to be given to the Proxy. This allows us to test the Proxy when UE4 is not on the local machine.

[RN] Changed the -UE4 option of the WebRTC Proxy from -UE4Port=<Port> to -UE4=<IP:Port>. Previously the IP was hard-coded to localhost; now you may have the UE4 and WebRTC Proxy instances on separate machines.

diff --git a/Engine/Source/Programs/PixelStreaming/WebRTCProxy/src/Conductor.cpp b/Engine/Source/Programs/PixelStreaming/WebRTCProxy/src/Conductor.cpp
--- a/Engine/Source/Programs/PixelStreaming/WebRTCProxy/src/Conductor.cpp
+++ b/Engine/Source/Programs/PixelStreaming/WebRTCProxy/src/Conductor.cpp
@@ -44,7 +44,7 @@ FConductor::FConductor()
 
 	ResetPeerConnectionConfig();
 
-	UE4Connection.Connect("127.0.0.1", PARAM_UE4Port);
+	UE4Connection.Connect(PARAM_UE4.first, PARAM_UE4.second);
 }
 
 FConductor::~FConductor()
diff --git a/Engine/Source/Programs/PixelStreaming/WebRTCProxy/src/WebRTCProxy.cpp b/Engine/Source/Programs/PixelStreaming/WebRTCProxy/src/WebRTCProxy.cpp
--- a/Engine/Source/Programs/PixelStreaming/WebRTCProxy/src/WebRTCProxy.cpp
+++ b/Engine/Source/Programs/PixelStreaming/WebRTCProxy/src/WebRTCProxy.cpp
@@ -26,8 +26,8 @@ The Cirrus server to connect to. If not specified. it defaults to 127.0.0.1:8888
 -StunServer=<IP:Port>\n\
 Stun server to use.\n\
 \n\
--UE4Port=<Port>\n\
-The port UE4 is listening on\n\
+-UE4=<IP:Port>\n\
+The address and port UE4 is listening on. If not specified it defaults to 127.0.0.1:8124\n\
 \n\
 -AutoSetBitrate\n\
 If specified, it forcibly sends a bitrate request to UE4 once a client gets\n\
@@ -53,8 +53,8 @@ If specified, it will use local time in logging, instead of UTC.\n\
 TCHAR GInternalProjectName[64] = TEXT("WebRTCProxy");
 IMPLEMENT_FOREIGN_ENGINE_DIR();
 
-std::pair<std::string, uint16_t> PARAM_Cirrus{"127.0.0.1", 8888};
-uint16_t PARAM_UE4Port = 8124;
+std::pair<std::string, uint16_t> PARAM_Cirrus{ "127.0.0.1", 8888 };
+std::pair<std::string, uint16_t> PARAM_UE4{ "127.0.0.1", 8124 };
 bool PARAM_PlanB = false;
 bool PARAM_DbgWindow_Proxy = true;
 bool PARAM_DbgWindow_WebRTC = true;
@@ -117,7 +117,10 @@ bool ParseParameters(int argc, char* argv[])
 		return false;
 	}
 
-	PARAM_UE4Port = Params.GetAsInt("UE4Port", 8124).second;
+	if (!ProcessAddressParameter("UE4", PARAM_UE4))
+	{
+		return false;
+	}
 
 	PARAM_PlanB = Params.Has("PlanB");
 
diff --git a/Engine/Source/Programs/PixelStreaming/WebRTCProxy/src/WebRTCProxyCommon.h b/Engine/Source/Programs/PixelStreaming/WebRTCProxy/src/WebRTCProxyCommon.h
--- a/Engine/Source/Programs/PixelStreaming/WebRTCProxy/src/WebRTCProxyCommon.h
+++ b/Engine/Source/Programs/PixelStreaming/WebRTCProxy/src/WebRTCProxyCommon.h
@@ -38,7 +38,7 @@ std::string GetExtension(const std::string& FullFilename, std::string* Basename)
 // Available parameters
 //
 extern std::pair<std::string, uint16_t> PARAM_Cirrus;
-extern uint16_t PARAM_UE4Port;
+extern std::pair<std::string, uint16_t> PARAM_UE4;
 extern bool PARAM_PlanB;
 extern bool PARAM_LocalTime;
 
