Copy of CL 4439167

* Fix overlay so it's grey allowing any white icons/text to be seen when video element shows a white image

diff --git a/Engine/Source/Programs/PixelStreaming/WebServers/SignallingWebServer/public/player.css b/Engine/Source/Programs/PixelStreaming/WebServers/SignallingWebServer/public/player.css
--- a/Engine/Source/Programs/PixelStreaming/WebServers/SignallingWebServer/public/player.css
+++ b/Engine/Source/Programs/PixelStreaming/WebServers/SignallingWebServer/public/player.css
@@ -123,7 +123,7 @@ video{
 	font-family: var(--inputFont);
 	width: 100%;
 	height: 100%;
-	background-color: rgba(255, 255, 255, 0.3);
+	background-color: rgba(100, 100, 100, 0.3);
 }
 
 /* State for element to be clickable */
