Fix for accidental junk checked in with CL# 4623722.

diff --git a/Engine/Source/Programs/PixelStreaming/WebRTCProxy/src/Logging.cpp b/Engine/Source/Programs/PixelStreaming/WebRTCProxy/src/Logging.cpp
--- a/Engine/Source/Programs/PixelStreaming/WebRTCProxy/src/Logging.cpp
+++ b/Engine/Source/Programs/PixelStreaming/WebRTCProxy/src/Logging.cpp
@@ -103,6 +103,4 @@ void ILogOutput::LogToAll(
 	{
 		Out->Log(File, Line, Category, Verbosity, Buf);
 	}
-
-	va_end()
 }
