Copied from Dev-Enterprise - Fixed some compile errors in the WebRTCProxy caused by increased warning level.

diff --git a/Engine/Plugins/Experimental/PixelStreaming/Source/PixelStreaming/Private/ProtocolDefs.h b/Engine/Plugins/Experimental/PixelStreaming/Source/PixelStreaming/Private/ProtocolDefs.h
--- a/Engine/Plugins/Experimental/PixelStreaming/Source/PixelStreaming/Private/ProtocolDefs.h
+++ b/Engine/Plugins/Experimental/PixelStreaming/Source/PixelStreaming/Private/ProtocolDefs.h
@@ -2,6 +2,7 @@
 
 #pragma once
 
+#include "Misc/AssertionMacros.h"
 
 namespace PixelStreamingProtocol
 {
diff --git a/Engine/Source/Programs/PixelStreaming/WebRTCProxy/src/CirrusConnection.cpp b/Engine/Source/Programs/PixelStreaming/WebRTCProxy/src/CirrusConnection.cpp
--- a/Engine/Source/Programs/PixelStreaming/WebRTCProxy/src/CirrusConnection.cpp
+++ b/Engine/Source/Programs/PixelStreaming/WebRTCProxy/src/CirrusConnection.cpp
@@ -8,8 +8,8 @@ using PixelStreamingProtocol::EProxyToCirrusMsg;
 using FMsgSize = uint32_t;
 
 FCirrusConnection::FCirrusConnection(ICirrusConnectionObserver& Observer)
-    : Connection("Cirrus", *this)
-    , Observer(Observer)
+    : Observer(Observer)
+    , Connection("Cirrus", *this)
 {}
 
 void FCirrusConnection::Connect(const std::string& IP, uint16_t Port)
diff --git a/Engine/Source/Programs/PixelStreaming/WebRTCProxy/src/Logging.cpp b/Engine/Source/Programs/PixelStreaming/WebRTCProxy/src/Logging.cpp
--- a/Engine/Source/Programs/PixelStreaming/WebRTCProxy/src/Logging.cpp
+++ b/Engine/Source/Programs/PixelStreaming/WebRTCProxy/src/Logging.cpp
@@ -27,9 +27,9 @@ const char* LogVerbosityToString(ELogVerbosity v)
 }
 
 FLogCategoryBase::FLogCategoryBase(const char* Name, ELogVerbosity Verbosity, ELogVerbosity CompileTimeVerbosity)
-    : Name(Name)
-    , Verbosity(Verbosity)
+    : Verbosity(Verbosity)
     , CompileTimeVerbosity(CompileTimeVerbosity)
+    , Name(Name)
 {
 }
 
diff --git a/Engine/Source/Programs/PixelStreaming/WebRTCProxy/src/WebRTCProxyPCH.h b/Engine/Source/Programs/PixelStreaming/WebRTCProxy/src/WebRTCProxyPCH.h
--- a/Engine/Source/Programs/PixelStreaming/WebRTCProxy/src/WebRTCProxyPCH.h
+++ b/Engine/Source/Programs/PixelStreaming/WebRTCProxy/src/WebRTCProxyPCH.h
@@ -23,6 +23,7 @@
 #pragma warning(push)
 #pragma warning(disable: 4582) // error C4582: 'webrtc::RTCError::string_message_': constructor is not implicitly called
 #pragma warning(disable: 4583) // error C4583 : 'rtc::Optional<std::string>::value_' : destructor is not implicitly called
+#pragma warning(disable : 4265) // warning C4265: 'webrtc::VideoCaptureExternal': class has virtual functions, but destructor is not virtual
 #endif
 //
 // WebRTC headers
