Prevent the Pixel Streaming Demo from trying to change the resolution to 4K when the user hasn't added the -NvEncH264ConfigLevel=NV_ENC_LEVEL_H264_52 UE4 command line option to UE4.

#rnx

diff --git a/Engine/Source/Programs/PixelStreaming/WebServers/SignallingWebServer/scripts/app.js b/Engine/Source/Programs/PixelStreaming/WebServers/SignallingWebServer/scripts/app.js
--- a/Engine/Source/Programs/PixelStreaming/WebServers/SignallingWebServer/scripts/app.js
+++ b/Engine/Source/Programs/PixelStreaming/WebServers/SignallingWebServer/scripts/app.js
@@ -14,6 +14,7 @@ var matchViewportResolution;
 var lastTimeResized = new Date().getTime();
 var resizeTimeout;
 
+var onDataChannelConnected;
 var responseEventListeners = new Map();
 
 var t0 = Date.now();
@@ -276,6 +277,9 @@ function setupWebRtcPlayer(htmlElement, clientConfig){
 
     webRtcPlayerObj.onDataChannelConnected = function(){
     	showTextOverlay('WebRTC connected, waiting for video');
+		if (onDataChannelConnected) {
+			onDataChannelConnected();
+		}
     }
 
 	webRtcPlayerObj.onDataChannelMessage = function (data) {
