Merge Pixel Streaming code //UE4/Dev-Enterprise@8224093 -> //UE4/Main/

Also merge build break fixes 8224117, 8224125, 8225640.

Allow Pixel Streaming to stop streaming and to send a freeze frame JPEG for the browser to display. This is for situations where the camera is not moving and nothing is animated in the scene, so the user will not perceive and difference, other than the JPEG being higher quality perhaps.
This is controlled by two new Blueprint functions: FreezeFrame and UnfreezeFrame. FreezeFrame may be supplied with a texture which will be used for the freeze frame. If this is not supplied then the back buffer will be captured and used as the freeze frame, acting as an impostor for the suspended video stream.

diff --git a/Engine/Plugins/Experimental/PixelStreaming/Source/PixelStreaming/Private/PixelStreamingFreezeFrame.cpp b/Engine/Plugins/Experimental/PixelStreaming/Source/PixelStreaming/Private/PixelStreamingFreezeFrame.cpp
new file mode 100644
--- /dev/null
+++ b/Engine/Plugins/Experimental/PixelStreaming/Source/PixelStreaming/Private/PixelStreamingFreezeFrame.cpp
@@ -0,0 +1,23 @@
+// Copyright 1998-2019 Epic Games, Inc. All Rights Reserved.
+
+#include "PixelStreamingFreezeFrame.h"
+#include "IPixelStreamingPlugin.h"
+
+UPixelStreamingFreezeFrame* UPixelStreamingFreezeFrame::Singleton = nullptr;
+
+void UPixelStreamingFreezeFrame::FreezeFrame(UTexture2D* Texture)
+{
+	Singleton->PixelStreamingPlugin->FreezeFrame(Texture);
+}
+
+void UPixelStreamingFreezeFrame::UnfreezeFrame()
+{
+	Singleton->PixelStreamingPlugin->UnfreezeFrame();
+}
+
+void UPixelStreamingFreezeFrame::CreateInstance()
+{
+	Singleton = NewObject<UPixelStreamingFreezeFrame>();
+	Singleton->AddToRoot();
+	Singleton->PixelStreamingPlugin = &FModuleManager::Get().GetModuleChecked<IPixelStreamingPlugin>("PixelStreaming");
+}
diff --git a/Engine/Plugins/Experimental/PixelStreaming/Source/PixelStreaming/Private/PixelStreamingFreezeFrame.h b/Engine/Plugins/Experimental/PixelStreaming/Source/PixelStreaming/Private/PixelStreamingFreezeFrame.h
new file mode 100644
--- /dev/null
+++ b/Engine/Plugins/Experimental/PixelStreaming/Source/PixelStreaming/Private/PixelStreamingFreezeFrame.h
@@ -0,0 +1,46 @@
+// Copyright 1998-2019 Epic Games, Inc. All Rights Reserved.
+
+#pragma once
+
+#include "PixelStreamingFreezeFrame.generated.h"
+
+class IPixelStreamingPlugin;
+
+/**
+ * This singleton object allows Pixel Streaming to be frozen and unfrozen from
+ * Blueprint. When frozen, a freeze frame (a still image) will be used by the
+ * browser instead of the video stream.
+ */
+UCLASS()
+class UPixelStreamingFreezeFrame : public UObject
+{
+	GENERATED_BODY()
+
+public:
+
+	/**
+	 * Freeze Pixel Streaming.
+	 * @param Texture - The freeze frame to display. If null then the back buffer is captured.
+	 */
+	UFUNCTION(BlueprintCallable, Category = "Pixel Streaming Freeze Frame")
+	static void FreezeFrame(UTexture2D* Texture);
+
+	/**
+	 * Unfreeze Pixel Streaming. 
+	 */
+	UFUNCTION(BlueprintCallable, Category = "Pixel Streaming Freeze Frame")
+	static void UnfreezeFrame();
+
+	/**
+	 * Create the singleton.
+	 */
+	static void CreateInstance();
+
+private:
+
+	// The singleton object.
+	static UPixelStreamingFreezeFrame* Singleton;
+
+	// For convenience we keep a reference to the Pixel Streaming plugin.
+	IPixelStreamingPlugin* PixelStreamingPlugin;
+};
diff --git a/Engine/Plugins/Experimental/PixelStreaming/Source/PixelStreaming/Private/PixelStreamingPlugin.cpp b/Engine/Plugins/Experimental/PixelStreaming/Source/PixelStreaming/Private/PixelStreamingPlugin.cpp
--- a/Engine/Plugins/Experimental/PixelStreaming/Source/PixelStreaming/Private/PixelStreamingPlugin.cpp
+++ b/Engine/Plugins/Experimental/PixelStreaming/Source/PixelStreaming/Private/PixelStreamingPlugin.cpp
@@ -16,18 +16,28 @@
 #include "Rendering/SlateRenderer.h"
 #include "Framework/Application/SlateApplication.h"
 #include "Misc/ConfigCacheIni.h"
+#include "PixelStreamingFreezeFrame.h"
 #include "PixelStreamingInputDevice.h"
 #include "PixelStreamingInputComponent.h"
 #include "GameFramework/GameModeBase.h"
 #include "Dom/JsonObject.h"
 #include "Misc/App.h"
 #include "Misc/MessageDialog.h"
+#include "IImageWrapper.h"
+#include "IImageWrapperModule.h"
+#include "Async/Async.h"
 
 DEFINE_LOG_CATEGORY(PixelStreaming);
 DEFINE_LOG_CATEGORY(PixelStreamingInput);
 DEFINE_LOG_CATEGORY(PixelStreamingNet);
 DEFINE_LOG_CATEGORY(PixelStreamingCapture);
 
+TAutoConsoleVariable<int32> CVarPixelStreamingFreezeFrameQuality(
+	TEXT("PixelStreaming.FreezeFrameQuality"),
+	100,
+	TEXT("Compression quality of the freeze frame"),
+	ECVF_Default);
+
 /** IModuleInterface implementation */
 void FPixelStreamingPlugin::StartupModule()
 {
@@ -71,6 +81,12 @@ void FPixelStreamingPlugin::StartupModule()
 	IModularFeatures::Get().RegisterModularFeature(GetModularFeatureName(), this);
 
 	FApp::SetUnfocusedVolumeMultiplier(1.0f);
+
+	// Allow Pixel Streaming to be frozen and a freeze frame image to be used
+	// instead of the video stream.
+	UPixelStreamingFreezeFrame::CreateInstance();
+	bFrozen = false;
+	bCaptureNextBackBufferAndStream = false;
 }
 
 void FPixelStreamingPlugin::ShutdownModule()
@@ -126,7 +142,25 @@ void FPixelStreamingPlugin::OnBackBufferReady_RenderThread(SWindow& SlateWindow,
 		Streamer = MakeUnique<FStreamer>(*IP, Port, BackBuffer);
 	}
 
-	Streamer->OnFrameBufferReady(BackBuffer);
+	if (!bFrozen)
+	{
+		Streamer->OnFrameBufferReady(BackBuffer);
+	}
+
+	// Check to see if we have been instructed to capture the back buffer as a
+	// freeze frame.
+	if (bCaptureNextBackBufferAndStream)
+	{
+		bCaptureNextBackBufferAndStream = false;
+
+		// Read the data out of the back buffer and send as a JPEG.
+		FRHICommandListImmediate& RHICmdList = FRHICommandListExecutor::GetImmediateCommandList();
+		FIntRect Rect(0, 0, BackBuffer->GetSizeX(), BackBuffer->GetSizeY());
+		TArray<FColor>* Data = new TArray<FColor>;
+
+		RHICmdList.ReadSurfaceData(BackBuffer, Rect, *Data, FReadSurfaceDataFlags());
+		SendJpeg(Data, Rect);
+	}
 }
 
 void FPixelStreamingPlugin::OnPreResizeWindowBackbuffer(void* BackBuffer)
@@ -166,6 +200,58 @@ TSharedPtr<FPixelStreamingInputDevice> FPixelStreamingPlugin::GetInputDevicePtr(
 	return InputDevice;
 }
 
+void FPixelStreamingPlugin::FreezeFrame(UTexture2D* Texture)
+{
+	if (Texture)
+	{
+		// A frame is supplied so immediately read its data and send as a JPEG.
+		FTexture2DRHIRef Texture2DRHI = Texture->Resource && Texture->Resource->TextureRHI ? Texture->Resource->TextureRHI->GetTexture2D() : nullptr;
+		if (!Texture2DRHI)
+		{
+			UE_LOG(PixelStreaming, Error, TEXT("Attempting freeze frame with texture %s with no texture 2D RHI"), *Texture->GetName());
+			return;
+		}
+
+		struct FReadSurfaceContext
+		{
+			FTexture2DRHIRef Texture;
+			FIntRect Rect;
+			TArray<FColor>* Data;
+		};
+
+		FReadSurfaceContext Context =
+		{
+			Texture2DRHI,
+			FIntRect(0, 0, Texture->GetSizeX(), Texture->GetSizeY()),
+			new TArray<FColor>
+		};
+
+		ENQUEUE_RENDER_COMMAND(ReadSurfaceCommand)(
+			[this, Context](FRHICommandListImmediate& RHICmdList)
+		{
+			RHICmdList.ReadSurfaceData(Context.Texture, Context.Rect, *Context.Data, FReadSurfaceDataFlags());
+			SendJpeg(Context.Data, Context.Rect);
+		});
+	}
+	else
+	{
+		// A frame is not supplied, so we need to capture the back buffer at
+		// the next opportunity, and send as a JPEG.
+		bCaptureNextBackBufferAndStream = true;
+	}
+
+	// Stop streaming.
+	bFrozen = true;
+}
+
+void FPixelStreamingPlugin::UnfreezeFrame()
+{
+	Streamer->SendUnfreezeFrame();
+
+	// Resume streaming.
+	bFrozen = false;
+}
+
 void FPixelStreamingPlugin::AddClientConfig(TSharedRef<FJsonObject>& JsonObject)
 {
 	checkf(InputDevice.IsValid(), TEXT("No Input Device available when populating Client Config"));
@@ -215,4 +301,28 @@ void FPixelStreamingPlugin::OnGameModeLogout(AGameModeBase* GameMode, AControlle
 	InputComponents.Empty();
 }
 
+void FPixelStreamingPlugin::SendJpeg(TArray<FColor>* Data, const FIntRect& Rect)
+{
+	AsyncTask(ENamedThreads::GameThread, [this, Data, Rect]
+	{
+		// Set up the image wrapper so we can compress the frame data to a JPEG.
+		IImageWrapperModule& ImageWrapperModule = FModuleManager::LoadModuleChecked<IImageWrapperModule>(FName("ImageWrapper"));
+		TSharedPtr<IImageWrapper> ImageWrapper = ImageWrapperModule.CreateImageWrapper(EImageFormat::JPEG);
+		bool bSuccess = ImageWrapper->SetRaw(Data->GetData(), Data->Num() * sizeof(FColor), Rect.Width(), Rect.Height(), ERGBFormat::BGRA, 8);
+		if (bSuccess)
+		{
+			// Compress to a JPEG of the maximum possible quality.
+			int32 Quality = CVarPixelStreamingFreezeFrameQuality.GetValueOnGameThread();
+			const TArray<uint8>& JpegBytes = ImageWrapper->GetCompressed(Quality);
+			Streamer->SendFreezeFrame(JpegBytes);
+		}
+		else
+		{
+			UE_LOG(PixelStreaming, Error, TEXT("JPEG image wrapper failed to accept frame data"));
+		}
+
+		delete Data;
+	});
+}
+
 IMPLEMENT_MODULE(FPixelStreamingPlugin, PixelStreaming)
diff --git a/Engine/Plugins/Experimental/PixelStreaming/Source/PixelStreaming/Private/PixelStreamingPlugin.h b/Engine/Plugins/Experimental/PixelStreaming/Source/PixelStreaming/Private/PixelStreamingPlugin.h
--- a/Engine/Plugins/Experimental/PixelStreaming/Source/PixelStreaming/Private/PixelStreamingPlugin.h
+++ b/Engine/Plugins/Experimental/PixelStreaming/Source/PixelStreaming/Private/PixelStreamingPlugin.h
@@ -38,6 +38,9 @@ public:
 	 */
 	TSharedPtr<FPixelStreamingInputDevice> GetInputDevicePtr();
 
+	virtual void FreezeFrame(UTexture2D* Texture) override;
+	virtual void UnfreezeFrame() override;
+
 private:
 
 	bool CheckPlatformCompatibility() const;
@@ -47,9 +50,12 @@ private:
 	void OnPreResizeWindowBackbuffer_RenderThread();
 	void OnGameModePostLogin(AGameModeBase* GameMode, APlayerController* NewPlayer);
 	void OnGameModeLogout(AGameModeBase* GameMode, AController* Exiting);
+	void SendJpeg(TArray<FColor>* Data, const FIntRect& Rect);
 
 	TUniquePtr<FStreamer> Streamer;
 	FTexture2DRHIRef mResolvedFrameBuffer;
 	TSharedPtr<FPixelStreamingInputDevice> InputDevice;
 	TArray<UPixelStreamingInputComponent*> InputComponents;
+	bool bFrozen;
+	bool bCaptureNextBackBufferAndStream;
 };
diff --git a/Engine/Plugins/Experimental/PixelStreaming/Source/PixelStreaming/Private/ProtocolDefs.h b/Engine/Plugins/Experimental/PixelStreaming/Source/PixelStreaming/Private/ProtocolDefs.h
--- a/Engine/Plugins/Experimental/PixelStreaming/Source/PixelStreaming/Private/ProtocolDefs.h
+++ b/Engine/Plugins/Experimental/PixelStreaming/Source/PixelStreaming/Private/ProtocolDefs.h
@@ -60,16 +60,16 @@ namespace PixelStreamingProtocol
 	};
 
 	// !!! modifying this enum make sure to update the next function !!!
-	enum class EToProxyMsg : uint8 { AudioPCM, SpsPps, VideoIDR, Video, ClientConfig, Response, Count };
+	enum class EToProxyMsg : uint8 { AudioPCM, SpsPps, VideoIDR, Video, ClientConfig, Response, FreezeFrame, UnfreezeFrame, Count };
 	inline const TCHAR* PacketTypeStr(EToProxyMsg PktType) 
 	{
-		static const TCHAR* Str[static_cast<uint8>(EToProxyMsg::Count)] = { TEXT("AudioPCM"), TEXT("SpsPps"), TEXT("VideoIDR"), TEXT("Video"), TEXT("ClientConfig"), TEXT("Response") };
+		static const TCHAR* Str[static_cast<uint8>(EToProxyMsg::Count)] = { TEXT("AudioPCM"), TEXT("SpsPps"), TEXT("VideoIDR"), TEXT("Video"), TEXT("ClientConfig"), TEXT("Response"), TEXT("FreezeFrame"), TEXT("UnfreezeFrame") };
 		check(PktType < EToProxyMsg::Count);
 		return Str[static_cast<uint8>(PktType)];
 	}
 
 	//! Messages that can be sent to the webrtc clients
-	enum class EToClientMsg : uint8 { QualityControlOwnership, Response };
+	enum class EToClientMsg : uint8 { QualityControlOwnership, Response, FreezeFrame, UnfreezeFrame };
 
 	enum class ECirrusToProxyMsg : uint8 { offer, iceCandidate, clientDisconnected, config, count };
 	enum class EProxyToCirrusMsg : uint8 { answer, iceCandidate, disconnectClient };
diff --git a/Engine/Plugins/Experimental/PixelStreaming/Source/PixelStreaming/Private/ProxyConnection.cpp b/Engine/Plugins/Experimental/PixelStreaming/Source/PixelStreaming/Private/ProxyConnection.cpp
--- a/Engine/Plugins/Experimental/PixelStreaming/Source/PixelStreaming/Private/ProxyConnection.cpp
+++ b/Engine/Plugins/Experimental/PixelStreaming/Source/PixelStreaming/Private/ProxyConnection.cpp
@@ -121,6 +121,11 @@ bool FProxyConnection::AcceptConnection(const FString& IP, uint16 Port)
 
 	UE_LOG(PixelStreamingNet, Log, TEXT("Accepted connection from WebRTC Proxy: %s"), *ProxyAddr->ToString(true));
 
+	// We may already have a freeze frame so we should send that over to the
+	// Proxy immediately, so it can be passed to the browser without waiting
+	// for video.
+	Streamer.SendFreezeFrame();
+
 	return true;
 }
 
diff --git a/Engine/Plugins/Experimental/PixelStreaming/Source/PixelStreaming/Private/Streamer.cpp b/Engine/Plugins/Experimental/PixelStreaming/Source/PixelStreaming/Private/Streamer.cpp
--- a/Engine/Plugins/Experimental/PixelStreaming/Source/PixelStreaming/Private/Streamer.cpp
+++ b/Engine/Plugins/Experimental/PixelStreaming/Source/PixelStreaming/Private/Streamer.cpp
@@ -202,8 +202,8 @@ void FStreamer::Stream(uint64 Timestamp, PixelStreamingProtocol::EToProxyMsg Pkt
 
 	SaveEncodedVideoToFile(PktType, Data, Size);
 
-	if (ProxyConnection->Send(reinterpret_cast<const uint8*>(&Timestamp), sizeof(Timestamp))
-		&& ProxyConnection->Send(reinterpret_cast<const uint8*>(&PktType), 1)
+	if (ProxyConnection->Send(reinterpret_cast<const uint8*>(&PktType), 1)
+		&& ProxyConnection->Send(reinterpret_cast<const uint8*>(&Timestamp), sizeof(Timestamp))
 		&& ProxyConnection->Send(reinterpret_cast<const uint8*>(&Size), sizeof(Size))
 		&& ProxyConnection->Send(Data, Size))
 	{
@@ -345,3 +345,26 @@ void FStreamer::SendResponse(const FString& Descriptor)
 {
 	Stream(FPlatformTime::Seconds(), PixelStreamingProtocol::EToProxyMsg::Response, reinterpret_cast<const uint8*>(*Descriptor), Descriptor.Len() * sizeof(TCHAR));
 }
+
+void FStreamer::SendFreezeFrame(const TArray<uint8>& JpegBytes)
+{
+	UE_LOG(PixelStreaming, Log, TEXT("Sending freeze frame to Proxy"));
+	Stream(FPlatformTime::Seconds(), PixelStreamingProtocol::EToProxyMsg::FreezeFrame, JpegBytes.GetData(), JpegBytes.Num());
+	CachedJpegBytes = JpegBytes;
+}
+
+void FStreamer::SendFreezeFrame()
+{
+	if (CachedJpegBytes.Num() > 0)
+	{
+		UE_LOG(PixelStreaming, Log, TEXT("Sending cached freeze frame to Proxy"));
+		Stream(FPlatformTime::Seconds(), PixelStreamingProtocol::EToProxyMsg::FreezeFrame, CachedJpegBytes.GetData(), CachedJpegBytes.Num());
+	}
+}
+
+void FStreamer::SendUnfreezeFrame()
+{
+	UE_LOG(PixelStreaming, Log, TEXT("Sending unfreeze message to Proxy"));
+	Stream(FPlatformTime::Seconds(), PixelStreamingProtocol::EToProxyMsg::UnfreezeFrame, nullptr, 0);
+	CachedJpegBytes.Empty();
+}
diff --git a/Engine/Plugins/Experimental/PixelStreaming/Source/PixelStreaming/Private/Streamer.h b/Engine/Plugins/Experimental/PixelStreaming/Source/PixelStreaming/Private/Streamer.h
--- a/Engine/Plugins/Experimental/PixelStreaming/Source/PixelStreaming/Private/Streamer.h
+++ b/Engine/Plugins/Experimental/PixelStreaming/Source/PixelStreaming/Private/Streamer.h
@@ -49,6 +49,9 @@ public:
 	void SetFramerate(int32 Fps);
 
 	void SendResponse(const FString& Descriptor);
+	void SendFreezeFrame(const TArray<uint8>& JpegBytes);
+	void SendFreezeFrame();
+	void SendUnfreezeFrame();
 
 private:
 	void CreateVideoEncoder(const FTexture2DRHIRef& FrameBuffer);
@@ -81,5 +84,9 @@ private:
 #endif
 
 	int32 InitialMaxFPS;
+
+	// When we send a freeze frame we retain the data to handle connection
+	// scenarios.
+	TArray<uint8> CachedJpegBytes;
 };
 
diff --git a/Engine/Plugins/Experimental/PixelStreaming/Source/PixelStreaming/Public/IPixelStreamingPlugin.h b/Engine/Plugins/Experimental/PixelStreaming/Source/PixelStreaming/Public/IPixelStreamingPlugin.h
--- a/Engine/Plugins/Experimental/PixelStreaming/Source/PixelStreaming/Public/IPixelStreamingPlugin.h
+++ b/Engine/Plugins/Experimental/PixelStreaming/Source/PixelStreaming/Public/IPixelStreamingPlugin.h
@@ -8,6 +8,8 @@
 #include "PixelStreamingCommon.h"
 #include "IInputDeviceModule.h"
 
+class UTexture2D;
+
 /**
 * The public interface to this module
 */
@@ -56,5 +58,16 @@ public:
 	 * @param Descriptor - A generic descriptor string.
 	 */
 	virtual void SendResponse(const FString& Descriptor) = 0;
+
+	/**
+	 * Freeze Pixel Streaming.
+	 * @param Texture - The freeze frame to display. If null then the back buffer is captured.
+	 */
+	virtual void FreezeFrame(UTexture2D* Texture) = 0;
+
+	/**
+	 * Unfreeze Pixel Streaming.
+	 */
+	virtual void UnfreezeFrame() = 0;
 };
 
diff --git a/Engine/Source/Programs/PixelStreaming/WebRTCProxy/src/ClientSession.cpp b/Engine/Source/Programs/PixelStreaming/WebRTCProxy/src/ClientSession.cpp
--- a/Engine/Source/Programs/PixelStreaming/WebRTCProxy/src/ClientSession.cpp
+++ b/Engine/Source/Programs/PixelStreaming/WebRTCProxy/src/ClientSession.cpp
@@ -96,6 +96,17 @@ void FClientSession::DisconnectClient()
 	Outer.CirrusConnection.SendDisconnectClient(ClientId);
 }
 
+void FClientSession::SendOnDataChannel(PixelStreamingProtocol::EToClientMsg ToClientMsg, const void* Pkt, uint32_t Size)
+{
+	if (DataChannel)
+	{
+		rtc::CopyOnWriteBuffer Buffer(Size + 1);
+		Buffer[0] = static_cast<uint8_t>(ToClientMsg);
+		std::memcpy(&Buffer[1], reinterpret_cast<const uint8_t*>(Pkt), Size);
+		DataChannel->Send(webrtc::DataBuffer(Buffer, true));
+	}
+}
+
 //
 // webrtc::PeerConnectionObserver implementation.
 //
@@ -191,6 +202,21 @@ void FClientSession::OnRemoveTrack(rtc::scoped_refptr<webrtc::RtpReceiverInterfa
 // webrtc::DataChannelObserver implementation.
 //
 
+void FClientSession::OnStateChange()
+{
+	webrtc::DataChannelInterface::DataState State = this->DataChannel->state();
+	if (State == webrtc::DataChannelInterface::DataState::kOpen)
+	{
+		// Once the data channel is opened, we check to see if we have any
+		// freeze frame chunks cached. If so then we send them immediately
+		// to the browser for display, without waiting for the video.
+		for (std::vector<uint8>& FreezeFrameChunk : Outer.FreezeFrameChunks)
+		{
+			SendOnDataChannel(PixelStreamingProtocol::EToClientMsg::FreezeFrame, FreezeFrameChunk.data(), FreezeFrameChunk.size());
+		}
+	}
+}
+
 void FClientSession::OnMessage(const webrtc::DataBuffer& Buffer)
 {
 	auto MsgType = static_cast<PixelStreamingProtocol::EToUE4Msg>(Buffer.data.data()[0]);
diff --git a/Engine/Source/Programs/PixelStreaming/WebRTCProxy/src/ClientSession.h b/Engine/Source/Programs/PixelStreaming/WebRTCProxy/src/ClientSession.h
--- a/Engine/Source/Programs/PixelStreaming/WebRTCProxy/src/ClientSession.h
+++ b/Engine/Source/Programs/PixelStreaming/WebRTCProxy/src/ClientSession.h
@@ -15,6 +15,7 @@ struct FClientSession
 	~FClientSession() override;
 
 	void DisconnectClient();
+	void SendOnDataChannel(PixelStreamingProtocol::EToClientMsg ToClientMsg, const void* Pkt, uint32_t Size);
 
 	FConductor& Outer;
 	FClientId ClientId;
@@ -43,8 +44,7 @@ struct FClientSession
 	//
 	// werbrtc::DataChannelObserver implementation.
 	//
-	void OnStateChange() override
-	{}
+	void OnStateChange() override;
 	void OnBufferedAmountChange(uint64_t PreviousAmount) override
 	{}
 	void OnMessage(const webrtc::DataBuffer& Buffer) override;
diff --git a/Engine/Source/Programs/PixelStreaming/WebRTCProxy/src/Conductor.cpp b/Engine/Source/Programs/PixelStreaming/WebRTCProxy/src/Conductor.cpp
--- a/Engine/Source/Programs/PixelStreaming/WebRTCProxy/src/Conductor.cpp
+++ b/Engine/Source/Programs/PixelStreaming/WebRTCProxy/src/Conductor.cpp
@@ -14,6 +14,8 @@ const char VideoLabel[] = "video_label";
 
 using webrtc::SdpType;
 
+const int32 FConductor::ChunkMaxSize = 16 * 1024;
+
 //////////////////////////////////////////////////////////////////////////
 // FConductor
 //////////////////////////////////////////////////////////////////////////
@@ -236,6 +238,7 @@ void FConductor::OnUE4Disconnected()
 {
 	DeleteAllClients();
 	CirrusConnection.Disconnect();
+	FreezeFrameChunks.clear();
 }
 
 void FConductor::OnUE4Packet(PixelStreamingProtocol::EToProxyMsg PktType, const void* Pkt, uint32_t Size)
@@ -248,18 +251,23 @@ void FConductor::OnUE4Packet(PixelStreamingProtocol::EToProxyMsg PktType, const
 	}
 	else if (PktType == PixelStreamingProtocol::EToProxyMsg::Response)
 	{
-		// Currently broadcast the response to all clients.
-		for (auto&& Client : Clients)
+		BroadcastClientMsg(PixelStreamingProtocol::EToClientMsg::Response, Pkt, Size);
+	}
+	else if (PktType == PixelStreamingProtocol::EToProxyMsg::FreezeFrame)
+	{
+		// The freeze frame JPEG data is too large to send over the data channel
+		// in one go, so split it into chunks and send all the chunks.
+		ChunkifyFreezeFrame(Pkt, Size);
+		for (std::vector<uint8>& FreezeFrameChunk : FreezeFrameChunks)
 		{
-			if (Client.second->DataChannel)
-			{
-				rtc::CopyOnWriteBuffer Buffer(Size + 1);
-				Buffer[0] = static_cast<uint8_t>(PixelStreamingProtocol::EToClientMsg::Response);
-				std::memcpy(&Buffer[1], reinterpret_cast<const uint8_t*>(Pkt), Size);
-				Client.second->DataChannel->Send(webrtc::DataBuffer(Buffer, true));
-			}
+			BroadcastClientMsg(PixelStreamingProtocol::EToClientMsg::FreezeFrame, FreezeFrameChunk.data(), FreezeFrameChunk.size());
 		}
 	}
+	else if (PktType == PixelStreamingProtocol::EToProxyMsg::UnfreezeFrame)
+	{
+		BroadcastClientMsg(PixelStreamingProtocol::EToClientMsg::UnfreezeFrame, Pkt, Size);
+		FreezeFrameChunks.clear();
+	}
 	else
 	{
 		check(VideoCapturer);
@@ -267,6 +275,40 @@ void FConductor::OnUE4Packet(PixelStreamingProtocol::EToProxyMsg PktType, const
 	}
 }
 
+void FConductor::BroadcastClientMsg(PixelStreamingProtocol::EToClientMsg ToClientMsg, const void* Pkt, uint32_t Size)
+{
+	// Currently broadcast the response to all clients.
+	for (auto&& Client : Clients)
+	{
+		Client.second->SendOnDataChannel(ToClientMsg, Pkt, Size);
+	}
+}
+
+void FConductor::ChunkifyFreezeFrame(const void* Pkt, uint32_t Size)
+{
+	// Split the freeze frame data into chunks so they may be sent over the data
+	// channel as individual messages and reassembled on the browser.
+	const uint8* BytePtr = static_cast<const uint8*>(Pkt);
+	int32 NumBytes = Size;
+	uint8 NumChunks = Size / ChunkMaxSize + (Size % ChunkMaxSize ? 1 : 0);
+
+	for (uint8 ChunkIdx = 0; ChunkIdx < NumChunks; ChunkIdx++)
+	{
+		int32 ChunkSize = NumBytes > ChunkMaxSize ? ChunkMaxSize : NumBytes;
+		NumBytes -= ChunkSize;
+
+		std::vector<uint8> Chunk;
+		Chunk.push_back(ChunkIdx);
+		if (ChunkIdx == 0)
+		{
+			Chunk.push_back(NumChunks);
+		}
+		Chunk.insert(Chunk.end(), BytePtr, BytePtr + ChunkSize);
+		BytePtr += ChunkSize;
+		FreezeFrameChunks.push_back(Chunk);
+	}
+}
+
 //
 // ICirrusConnectionObserver implementation.
 //
diff --git a/Engine/Source/Programs/PixelStreaming/WebRTCProxy/src/Conductor.h b/Engine/Source/Programs/PixelStreaming/WebRTCProxy/src/Conductor.h
--- a/Engine/Source/Programs/PixelStreaming/WebRTCProxy/src/Conductor.h
+++ b/Engine/Source/Programs/PixelStreaming/WebRTCProxy/src/Conductor.h
@@ -50,6 +50,8 @@ private:
 
 	void ResetPeerConnectionConfig();
 
+	void BroadcastClientMsg(PixelStreamingProtocol::EToClientMsg ToClientMsg, const void* Pkt, uint32_t Size);
+
 	friend FClientSession;
 
 private:
@@ -72,4 +74,15 @@ private:
 	// These are used only if using UnifiedPlan semantics
 	rtc::scoped_refptr<webrtc::AudioTrackInterface> AudioTrack;
 	rtc::scoped_refptr<webrtc::VideoTrackInterface> VideoTrack;
+
+	// The size of the buffer used by the data channel.
+	static const int32 ChunkMaxSize;
+
+	// A freeze frame JPEG is too large to send over the data channel in one go,
+	// so we must chunkify it.
+	std::vector<std::vector<uint8>> FreezeFrameChunks;
+
+	// Convert the given freeze frame JPEG data into chunks so each chunk can
+	// be sent separately over the data channel.
+	void ChunkifyFreezeFrame(const void* Pkt, uint32_t Size);
 };
diff --git a/Engine/Source/Programs/PixelStreaming/WebRTCProxy/src/UE4Connection.cpp b/Engine/Source/Programs/PixelStreaming/WebRTCProxy/src/UE4Connection.cpp
--- a/Engine/Source/Programs/PixelStreaming/WebRTCProxy/src/UE4Connection.cpp
+++ b/Engine/Source/Programs/PixelStreaming/WebRTCProxy/src/UE4Connection.cpp
@@ -75,23 +75,26 @@ void FUE4Connection::Send(const void* Data, uint32_t Size)
 
 uint32_t FUE4Connection::OnRead(const uint8_t* Data, uint32_t Size)
 {
-	if (!bStreamingStarted)
-		return Size; // drop data as there's no clients to receive it
+	if (Size < sizeof(EToProxyMsg))
+		return 0;
+
+	const uint8_t* Ptr = Data;  // pointer to current read pos in the buffer
+
+	auto PktType = *reinterpret_cast<const EToProxyMsg*>(Ptr);
+	Ptr += sizeof(PktType);
+
+	if (DropPacket(PktType))
+		return Size; // drop data
 
 	using FTimestamp = uint64_t;
 	using FPayloadSize = uint32_t;
 
-	if (Size < sizeof(FTimestamp) + sizeof(EToProxyMsg) + sizeof(FPayloadSize))
+	if (Size < sizeof(EToProxyMsg) + sizeof(FTimestamp) + sizeof(FPayloadSize))
 		return 0;
 
-	const uint8_t* Ptr = Data;  // pointer to current read pos in the buffer
-
 	auto CaptureTimeMs = *reinterpret_cast<const FTimestamp*>(Ptr);
 	Ptr += sizeof(CaptureTimeMs);
 
-	auto PktType = *reinterpret_cast<const EToProxyMsg*>(Ptr);
-	Ptr += sizeof(PktType);
-
 	auto PayloadSize = *reinterpret_cast<const FPayloadSize*>(Ptr);
 	Ptr += sizeof(PayloadSize);
 
@@ -104,3 +107,15 @@ uint32_t FUE4Connection::OnRead(const uint8_t* Data, uint32_t Size)
 
 	return static_cast<uint32_t>(Ptr - Data);
 }
+
+bool FUE4Connection::DropPacket(EToProxyMsg PktType) const
+{
+	if (PktType == EToProxyMsg::FreezeFrame || PktType == EToProxyMsg::UnfreezeFrame)
+	{
+		return false; // We never drop freeze frame data as it's used when the video isn't streaming.
+	}
+	else
+	{
+		return !bStreamingStarted; // drop data as there's no clients to receive it
+	}
+}
diff --git a/Engine/Source/Programs/PixelStreaming/WebRTCProxy/src/UE4Connection.h b/Engine/Source/Programs/PixelStreaming/WebRTCProxy/src/UE4Connection.h
--- a/Engine/Source/Programs/PixelStreaming/WebRTCProxy/src/UE4Connection.h
+++ b/Engine/Source/Programs/PixelStreaming/WebRTCProxy/src/UE4Connection.h
@@ -44,4 +44,7 @@ private:
 	IUE4ConnectionObserver& Observer;
 	FAsyncConnection Connection;
 	std::atomic<bool> bStreamingStarted = false;
+
+	// Check to see if we should drop a packet of the given message type.
+	bool DropPacket(PixelStreamingProtocol::EToProxyMsg PktType) const;
 };
diff --git a/Engine/Source/Programs/PixelStreaming/WebServers/SignallingWebServer/scripts/app.js b/Engine/Source/Programs/PixelStreaming/WebServers/SignallingWebServer/scripts/app.js
--- a/Engine/Source/Programs/PixelStreaming/WebServers/SignallingWebServer/scripts/app.js
+++ b/Engine/Source/Programs/PixelStreaming/WebServers/SignallingWebServer/scripts/app.js
@@ -17,6 +17,18 @@ var resizeTimeout;
 var onDataChannelConnected;
 var responseEventListeners = new Map();
 
+var freezeFrameOverlay = null;
+var shouldShowPlayOverlay = true;
+
+// A freeze frame is a still JPEG image shown instead of the video.
+var freezeFrame = {
+	chunks: null,
+	numChunks: 0,
+	height: 0,
+	width: 0,
+	valid: false
+};
+
 var t0 = Date.now();
 function log(str) {
     console.log(`${Math.floor(Date.now() - t0)}: ` + str);
@@ -220,8 +232,10 @@ function showPlayOverlay(){
 
 		requestQualityControl();
 
+		showFreezeFrameOverlay();
 		hideOverlay();
 	});
+	shouldShowPlayOverlay = false;
 }
 
 function hideOverlay(){
@@ -255,12 +269,15 @@ function removeResponseEventListener(name) {
 // Must be kept in sync with PixelStreamingProtocol::EToClientMsg C++ enum.
 const ToClientMessageType = {
 	QualityControlOwnership: 0,
-	Response: 1
+	Response: 1,
+	FreezeFrame: 2,
+	UnfreezeFrame: 3
 };
 
 function setupWebRtcPlayer(htmlElement, clientConfig){
     webRtcPlayerObj = new webRtcPlayer({peerConnectionOptions: clientConfig.peerConnectionOptions});
     htmlElement.appendChild(webRtcPlayerObj.video);
+	htmlElement.appendChild(freezeFrameOverlay);
 
     webRtcPlayerObj.onWebRtcOffer = function (offer) {
    		socket.emit("webrtc-offer", offer);
@@ -270,9 +287,11 @@ function setupWebRtcPlayer(htmlElement, clientConfig){
 		socket.emit('webrtc-ice', candidate);
     };
 
-    webRtcPlayerObj.onVideoInitialised = function(){
-    	showPlayOverlay();
-    	resizePlayerStyle();
+	webRtcPlayerObj.onVideoInitialised = function(){
+		if (shouldShowPlayOverlay) {
+			showPlayOverlay();
+			resizePlayerStyle();
+		}
     }
 
     webRtcPlayerObj.onDataChannelConnected = function(){
@@ -298,6 +317,42 @@ function setupWebRtcPlayer(htmlElement, clientConfig){
 			for (let listener of responseEventListeners.values()) {
 				listener(response);
 			}
+		} else if (view[0] == ToClientMessageType.FreezeFrame) {
+			let chunkIdx = view[1];
+			if (chunkIdx == 0) {
+				// The first chunk has arrived.
+				freezeFrame.numChunks = view[2];
+				freezeFrame.chunks = new Uint8Array(data.slice(3));
+			} else {
+				// Subsequent chunks expand the data.
+				let chunk = new Uint8Array(data.slice(2));
+				let expandedChunks = new Uint8Array(freezeFrame.chunks.length + chunk.length);
+				expandedChunks.set(freezeFrame.chunks);
+				expandedChunks.set(chunk, freezeFrame.chunks.length);
+				freezeFrame.chunks = expandedChunks;
+			}
+
+			if (chunkIdx == freezeFrame.numChunks - 1) {
+				// All the chunks have arrived so we set the JPEG data.
+				freezeFrame.valid = true;
+				let base64 = btoa(freezeFrame.chunks.reduce((data, byte) => data + String.fromCharCode(byte), ''));
+				freezeFrameOverlay.src = 'data:image/jpeg;base64,' + base64;
+
+				freezeFrameOverlay.onload = function () {
+					freezeFrame.height = freezeFrameOverlay.naturalHeight;
+					freezeFrame.width = freezeFrameOverlay.naturalWidth;
+					resizeFreezeFrameOverlay();
+
+					if (shouldShowPlayOverlay) {
+						showPlayOverlay();
+						resizePlayerStyle();
+					} else {
+						showFreezeFrameOverlay();
+					}
+				}
+			}
+		} else if (view[0] == ToClientMessageType.UnfreezeFrame) {
+			invalidateFreezeFrameOverlay();
 		}
 	}
 
@@ -456,13 +511,15 @@ function resizePlayerStyleToFillWindow(playerElement) {
 function resizePlayerStyleToActualSize(playerElement) {
     let videoElement = playerElement.getElementsByTagName("VIDEO");
 
-    // Display image in its actual size
-    styleWidth = videoElement.videoWidth;
-    styleHeight = videoElement.videoHeight;
-    styleTop = Math.floor((window.innerHeight - styleHeight) * 0.5);
-    styleLeft = Math.floor((window.innerWidth - styleWidth) * 0.5);
-    //Video is now 100% of the playerElement, so set the playerElement style
-    playerElement.style = "top: " + styleTop + "px; left: " + styleLeft + "px; width: " + styleWidth + "px; height: " + styleHeight + "px; cursor: " + styleCursor + "; " + styleAdditional;
+	if (videoElement.length > 0) {
+		// Display image in its actual size
+		styleWidth = videoElement[0].videoWidth;
+		styleHeight = videoElement[0].videoHeight;
+		styleTop = Math.floor((window.innerHeight - styleHeight) * 0.5);
+		styleLeft = Math.floor((window.innerWidth - styleWidth) * 0.5);
+		//Video is now 100% of the playerElement, so set the playerElement style
+		playerElement.style = "top: " + styleTop + "px; left: " + styleLeft + "px; width: " + styleWidth + "px; height: " + styleHeight + "px; cursor: " + styleCursor + "; " + styleAdditional;
+	}
 }
 
 function resizePlayerStyleToArbitrarySize(playerElement) {
@@ -471,6 +528,60 @@ function resizePlayerStyleToArbitrarySize(playerElement) {
     playerElement.style = "top: 0px; left: 0px; width: " + styleWidth + "px; height: " + styleHeight + "px; cursor: " + styleCursor + "; " + styleAdditional;
 }
 
+function setupFreezeFrameOverlay() {
+	freezeFrameOverlay = document.createElement('img');
+	freezeFrameOverlay.id = 'freezeFrameOverlay';
+	freezeFrameOverlay.style.display = 'none';
+	freezeFrameOverlay.style.pointerEvents = 'none';
+	freezeFrameOverlay.style.position = 'absolute';
+	freezeFrameOverlay.style.zIndex = '30';
+}
+
+function showFreezeFrameOverlay() {
+	if (freezeFrame.valid) {
+		freezeFrameOverlay.style.display = 'block';
+	}
+}
+
+function invalidateFreezeFrameOverlay() {
+	freezeFrameOverlay.style.display = 'none';
+	freezeFrame.valid = false;
+}
+
+function resizeFreezeFrameOverlay() {
+	if (freezeFrame.width != 0 && freezeFrame.height != 0) {
+		let displayWidth = 0;
+		let displayHeight = 0;
+		let displayTop = 0;
+		let displayLeft = 0;
+		let checkBox = document.getElementById('enlarge-display-to-fill-window-tgl');
+		if (checkBox != null && checkBox.checked) {
+			let windowAspectRatio = window.innerWidth / window.innerHeight;
+			let videoAspectRatio = freezeFrame.width / freezeFrame.height;
+			if (windowAspectRatio < videoAspectRatio) {
+				displayWidth = window.innerWidth;
+				displayHeight = Math.floor(window.innerWidth / videoAspectRatio);
+				displayTop = Math.floor((window.innerHeight - displayHeight) * 0.5);
+				displayLeft = 0;
+			} else {
+				displayWidth = Math.floor(window.innerHeight * videoAspectRatio);
+				displayHeight = window.innerHeight;
+				displayTop = 0;
+				displayLeft = Math.floor((window.innerWidth - displayWidth) * 0.5);
+			}
+		} else {
+			displayWidth = freezeFrame.width;
+			displayHeight = freezeFrame.height;
+			displayTop = 0;
+			displayLeft = 0;
+		}
+		freezeFrameOverlay.style.width = displayWidth + 'px';
+		freezeFrameOverlay.style.height = displayHeight + 'px';
+		freezeFrameOverlay.style.left = displayLeft + 'px';
+		freezeFrameOverlay.style.top = displayTop + 'px';
+	}
+}
+
 function resizePlayerStyle(event) {
 	var playerElement = document.getElementById('player');
 	
@@ -498,6 +609,7 @@ function resizePlayerStyle(event) {
 	} else {
 		resizePlayerStyleToArbitrarySize(playerElement);
 	}
+	resizeFreezeFrameOverlay();
 }
 
 function updateVideoStreamSize() {
@@ -1149,6 +1261,8 @@ function start() {
 				
 	if (!connect_on_load || is_reconnection){
 		showConnectOverlay();
+		invalidateFreezeFrameOverlay();
+		shouldShowPlayOverlay = true;
 		resizePlayerStyle();
 	} else {
 		connect();
@@ -1237,6 +1351,7 @@ function onClientConfig(clientConfig) {
 
 function load() {
 	setupHtmlEvents();
+	setupFreezeFrameOverlay();
     registerKeyboardEvents();
     start();
 }
