Merging 9546470 from //UE4/Dev-CharacterTech to //UE4/Release-4.24

For Pixel Streaming added a client-side AFK system which will disconnect inactive users. It will first warn the user after a period of inactivty that they are to be disconnected and a countdown will be shown. If the user does not click then they will be disconnected.
This is configurated in app.js as follows:
afk.enabled: Set this to true to activate the AFK system; default false.
afk.warnTimeout: Set this to the number of seconds of inactivity before the warning overlay is displayed; default 120 seconds.
afk.closeTimeout: Set this to the number of seconds the warning overlay will be displayed before the user is disconnected (unless they click); default 10 seconds.

[RN] Added a client-side AFK system which will disconnect inactive users.

diff --git a/Engine/Source/Programs/PixelStreaming/WebServers/SignallingWebServer/public/player.css b/Engine/Source/Programs/PixelStreaming/WebServers/SignallingWebServer/public/player.css
--- a/Engine/Source/Programs/PixelStreaming/WebServers/SignallingWebServer/public/player.css
+++ b/Engine/Source/Programs/PixelStreaming/WebServers/SignallingWebServer/public/player.css
@@ -123,7 +123,7 @@ video{
 	font-family: var(--inputFont);
 	width: 100%;
 	height: 100%;
-	background-color: rgba(100, 100, 100, 0.3);
+	background-color: rgba(100, 100, 100, 0.7);
 }
 
 /* State for element to be clickable */
diff --git a/Engine/Source/Programs/PixelStreaming/WebServers/SignallingWebServer/scripts/app.js b/Engine/Source/Programs/PixelStreaming/WebServers/SignallingWebServer/scripts/app.js
--- a/Engine/Source/Programs/PixelStreaming/WebServers/SignallingWebServer/scripts/app.js
+++ b/Engine/Source/Programs/PixelStreaming/WebServers/SignallingWebServer/scripts/app.js
@@ -30,6 +30,19 @@ var freezeFrame = {
 	valid: false
 };
 
+// Optionally detect if the user is not interacting (AFK) and disconnect them.
+var afk = {
+	enabled: false,   // Set to true to enable the AFK system.
+	warnTimeout: 120,   // The time to elapse before warning the user they are inactive.
+	closeTimeout: 10,   // The time after the warning when we disconnect the user.
+
+	active: false,   // Whether the AFK system is currently looking for inactivity.
+	overlay: undefined,   // The UI overlay warning the user that they are inactive.
+	warnTimer: undefined,   // The timer which waits to show the inactivity warning overlay.
+	countdown: 0,   // The inactivity warning overlay has a countdown to show time until disconnect.
+	countdownTimer: undefined,   // The timer used to tick the seconds shown on the inactivity warning overlay.
+}
+
 // If the user focuses on a UE4 input widget then we show them a button to open
 // the on-screen keyboard. JavaScript security means we can only show the
 // on-screen keyboard in response to a user interaction.
@@ -222,6 +235,7 @@ function showConnectOverlay() {
 
 	setOverlay('clickableState', startText, event => {
 		connect();
+		startAfkWarningTimer();
 	});
 }
 
@@ -249,10 +263,69 @@ function showPlayOverlay() {
 	shouldShowPlayOverlay = false;
 }
 
+function updateAfkOverlayText() {
+	afk.overlay.innerHTML = '<center>No activity detected<br>Disconnecting in ' + afk.countdown + ' seconds<br>Click to continue<br></center>';
+}
+
+function showAfkOverlay() {
+	// Pause the timer while the user is looking at the inactivity warning overlay.
+	stopAfkWarningTimer();
+
+	// Show the inactivity warning overlay.
+	afk.overlay = document.createElement('div');
+	afk.overlay.id = 'afkOverlay';
+	setOverlay('clickableState', afk.overlay, event => {
+		// The user clicked so start the timer again and carry on.
+		hideOverlay();
+		clearInterval(afk.countdownTimer);
+		startAfkWarningTimer();
+	});
+
+	afk.countdown = afk.closeTimeout;
+	updateAfkOverlayText();
+
+	if (inputOptions.controlScheme == ControlSchemeType.LockedMouse) {
+		document.exitPointerLock();
+	}
+
+	afk.countdownTimer = setInterval(function () {
+		afk.countdown--;
+		if (afk.countdown == 0) {
+			// The user failed to click so disconnect them.
+			hideOverlay();
+			ws.close();
+		} else {
+			// Update the countdown message.
+			updateAfkOverlayText();
+		}
+	}, 1000);
+}
+
 function hideOverlay() {
 	setOverlay('hiddenState');
 }
 
+// Start a timer which when elapsed will warn the user they are inactive.
+function startAfkWarningTimer() {
+	afk.active = afk.enabled;
+	resetAfkWarningTimer();
+}
+
+// Stop the timer which when elapsed will warn the user they are inactive.
+function stopAfkWarningTimer() {
+	afk.active = false;
+}
+
+// If the user interacts then reset the warning timer.
+function resetAfkWarningTimer() {
+	if (afk.active) {
+		clearTimeout(afk.warnTimer);
+		afk.warnTimer = setTimeout(function () {
+			showAfkOverlay();
+		}, afk.warnTimeout * 1000);
+	}
+}
+
 function createWebRtcOffer() {
 	if (webRtcPlayerObj) {
 		console.log('Creating offer');
@@ -265,8 +338,10 @@ function createWebRtcOffer() {
 }
 
 function sendInputData(data) {
-	if (webRtcPlayerObj)
+	if (webRtcPlayerObj) {
+		resetAfkWarningTimer();
 		webRtcPlayerObj.send(data);
+	}
 }
 
 function addResponseEventListener(name, listener) {
@@ -388,9 +463,9 @@ function setupWebRtcPlayer(htmlElement, config) {
 			}
 		} else if (view[0] === ToClientMessageType.UnfreezeFrame) {
 			invalidateFreezeFrameOverlay();
-		} else {
+					} else {
 			console.error(`unrecognised data received, packet ID ${view[0]}`);
-		}
+	}
 	};
 
 	registerInputs(webRtcPlayerObj.video);
