* Fixed wrong mouse pointer position offset when flipping the "Enlarge Display to Fill Window" toggle

#rnx

diff --git a/Engine/Source/Programs/PixelStreaming/WebServers/SignallingWebServer/scripts/app.js b/Engine/Source/Programs/PixelStreaming/WebServers/SignallingWebServer/scripts/app.js
--- a/Engine/Source/Programs/PixelStreaming/WebServers/SignallingWebServer/scripts/app.js
+++ b/Engine/Source/Programs/PixelStreaming/WebServers/SignallingWebServer/scripts/app.js
@@ -705,11 +705,6 @@ function resizePlayerStyle(event) {
 
 	updateVideoStreamSize();
 
-	// Calculating and normalizing positions depends on the width and height of
-	// the player.
-	playerElementClientRect = playerElement.getBoundingClientRect();
-	setupNormalizeAndQuantize();
-
 	if (playerElement.classList.contains('fixed-size'))
 		return;
 
@@ -724,6 +719,11 @@ function resizePlayerStyle(event) {
 	} else {
 		resizePlayerStyleToArbitrarySize(playerElement);
 	}
+	
+	// Calculating and normalizing positions depends on the width and height of
+	// the player.
+	playerElementClientRect = playerElement.getBoundingClientRect();
+	setupNormalizeAndQuantize();
 	resizeFreezeFrameOverlay();
 }
 
@@ -1214,6 +1214,7 @@ function registerLockedMouseEvents(playerElement) {
 // passes over the browser.
 function registerHoveringMouseEvents(playerElement) {
 	styleCursor = 'none';   // We will rely on UE4 client's software cursor.
+	//styleCursor = 'default';  // Showing cursor
 
 	playerElement.onmousemove = function (e) {
 		emitMouseMove(e.offsetX, e.offsetY, e.movementX, e.movementY);
