Merging 11279208

For the Pixel Streaming Signalling Web Server, add an error code for the disconnect player event.

diff --git a/Engine/Source/Programs/PixelStreaming/WebServers/SignallingWebServer/cirrus.js b/Engine/Source/Programs/PixelStreaming/WebServers/SignallingWebServer/cirrus.js
--- a/Engine/Source/Programs/PixelStreaming/WebServers/SignallingWebServer/cirrus.js
+++ b/Engine/Source/Programs/PixelStreaming/WebServers/SignallingWebServer/cirrus.js
@@ -297,7 +297,7 @@ streamerServer.on('connection', function (ws, req) {
 		} else if (msg.type == 'iceCandidate') {
 			player.ws.send(JSON.stringify(msg));
 		} else if (msg.type == 'disconnectPlayer') {
-			player.ws.close(msg.reason);
+			player.ws.close(1011 /* internal error */, msg.reason);
 		} else {
 			console.error(`unsupported Streamer message type: ${msg.type}`);
 			streamer.close(1008, 'Unsupported message type');
