Fix copyrights

diff --git a/Engine/Source/Programs/PixelStreaming/WebServers/SignallingWebServer/public/player.html b/Engine/Source/Programs/PixelStreaming/WebServers/SignallingWebServer/public/player.html
--- a/Engine/Source/Programs/PixelStreaming/WebServers/SignallingWebServer/public/player.html
+++ b/Engine/Source/Programs/PixelStreaming/WebServers/SignallingWebServer/public/player.html
@@ -1,4 +1,4 @@
-<!-- Copyright 1998-2018 Epic Games, Inc. All Rights Reserved. -->
+<!-- Copyright Epic Games, Inc. All Rights Reserved. -->
 <!DOCTYPE HTML>
 <html>
 <head>
diff --git a/Engine/Source/Runtime/AVEncoder/Private/Encoders/Amf/Amf_Common.cpp b/Engine/Source/Runtime/AVEncoder/Private/Encoders/Amf/Amf_Common.cpp
--- a/Engine/Source/Runtime/AVEncoder/Private/Encoders/Amf/Amf_Common.cpp
+++ b/Engine/Source/Runtime/AVEncoder/Private/Encoders/Amf/Amf_Common.cpp
@@ -1,3 +1,5 @@
+// Copyright Epic Games, Inc. All Rights Reserved.
+
 #include "Amf_Common.h"
 #include "AVEncoder.h"
 
